Adding Archvis template to release 4.24.

Test performed:

made a new project
selected the template
hit play
made a new level
everything is ok, sun sky system plays during PIE

diff --git a/Templates/TP_ArchvisBP/Config/DefaultEngine.ini b/Templates/TP_ArchvisBP/Config/DefaultEngine.ini
new file mode 100644
--- /dev/null
+++ b/Templates/TP_ArchvisBP/Config/DefaultEngine.ini
@@ -0,0 +1,93 @@
+[/Script/HardwareTargeting.HardwareTargetingSettings]
+; quality settings
+TargetedHardwareClass=Desktop
+AppliedTargetedHardwareClass=Desktop
+DefaultGraphicsPerformance=Maximum
+AppliedDefaultGraphicsPerformance=Maximum
+
+[/Script/Engine.RendererSettings]
+; exposure control and tone mapping settings
+r.Tonemapper.Sharpen=2
+r.TemporalAA.Upsampling=True
+r.DefaultFeature.AutoExposure.ExtendDefaultLuminanceRange=True
+
+;dynamic gi settings
+r.LightPropagationVolume=0
+
+; quality settings
+r.ReflectionCaptureResolution=2048
+r.AllowStaticLighting=True
+r.HighResScreenshotDelay 8
+r.DefaultBackBufferPixelFormat=4
+r.AllowGlobalClipPlane=False
+
+; shader settings
+r.SupportSkyAtmosphereAffectsHeightFog=True
+r.DefaultFeature.LensFlare=True
+r.ClearCoatNormal=False
+r.NormalMapsForStaticLighting=False
+
+
+[/Script/EngineSettings.GameMapsSettings]
+;defines what level to be loaded by default in the editor and at playback time
+GameDefaultMap=/Game/ArchvisProject/Main.Main
+EditorStartupMap=/Game/ArchvisProject/Main.Main
+GlobalDefaultGameMode=/Game/ArchvisProject/ArchvisGameMode.ArchvisGameMode_C
+
+[/Script/UnrealEd.UnrealEdEngine]
+; Remove Engine Template maps
+-TemplateMapInfos=(ThumbnailTexture=Texture2D'/Engine/Maps/Templates/Thumbnails/Default.Default',Map="/Engine/Maps/Templates/Template_Default")
+-TemplateMapInfos=(ThumbnailTexture=Texture2D'/Engine/Maps/Templates/Thumbnails/TimeOfDay.TimeOfDay',Map="/Engine/Maps/Templates/TimeOfDay_Default")
+-TemplateMapInfos=(ThumbnailTexture=Texture2D'/Engine/Maps/Templates/Thumbnails/VR-Basic.VR-Basic',Map="/Engine/Maps/Templates/VR-Basic")
+
+; Add Predefined Template Maps that are suitable for this template
+;+TemplateMapInfos=(ThumbnailTexture=Texture2D'/Game/Developers/Maps/Thumbnails/AssetName.AssetName',Map="/Game/Maps/Templates/MapName")
++TemplateMapInfos=(ThumbnailTexture=Texture2D'/Game/LevelTemplates/Default/Thumbnails/Default.Default',Map="/Game/LevelTemplates/Default/Default")
+
+[/Script/Engine.PhysicsSettings]
+DefaultGravityZ=-980.000000
+DefaultTerminalVelocity=4000.000000
+DefaultFluidFriction=0.300000
+SimulateScratchMemorySize=262144
+RagdollAggregateThreshold=4
+TriangleMeshTriangleMinAreaThreshold=5.000000
+bEnableShapeSharing=False
+bEnablePCM=True
+bEnableStabilization=False
+bWarnMissingLocks=True
+bEnable2DPhysics=False
+PhysicErrorCorrection=(PingExtrapolation=0.100000,PingLimit=100.000000,ErrorPerLinearDifference=1.000000,ErrorPerAngularDifference=1.000000,MaxRestoredStateError=1.000000,MaxLinearHardSnapDistance=400.000000,PositionLerp=0.000000,AngleLerp=0.400000,LinearVelocityCoefficient=100.000000,AngularVelocityCoefficient=10.000000,ErrorAccumulationSeconds=0.500000,ErrorAccumulationDistanceSq=15.000000,ErrorAccumulationSimilarity=100.000000)
+LockedAxis=Invalid
+DefaultDegreesOfFreedom=Full3D
+BounceThresholdVelocity=200.000000
+FrictionCombineMode=Average
+RestitutionCombineMode=Average
+MaxAngularVelocity=3600.000000
+MaxDepenetrationVelocity=0.000000
+ContactOffsetMultiplier=0.020000
+MinContactOffset=2.000000
+MaxContactOffset=8.000000
+bSimulateSkeletalMeshOnDedicatedServer=True
+DefaultShapeComplexity=CTF_UseSimpleAndComplex
+bDefaultHasComplexCollision=True
+bSuppressFaceRemapTable=False
+bSupportUVFromHitResults=False
+bDisableActiveActors=False
+bDisableKinematicStaticPairs=False
+bDisableKinematicKinematicPairs=False
+bDisableCCD=False
+bEnableEnhancedDeterminism=False
+AnimPhysicsMinDeltaTime=0.000000
+bSimulateAnimPhysicsAfterReset=False
+MaxPhysicsDeltaTime=0.033333
+bSubstepping=False
+bSubsteppingAsync=False
+MaxSubstepDeltaTime=0.016667
+MaxSubsteps=6
+SyncSceneSmoothingFactor=0.000000
+InitialAverageFrameRate=0.016667
+PhysXTreeRebuildRate=10
+DefaultBroadphaseSettings=(bUseMBPOnClient=False,bUseMBPOnServer=False,bUseMBPOuterBounds=False,MBPBounds=(Min=(X=0.000000,Y=0.000000,Z=0.000000),Max=(X=0.000000,Y=0.000000,Z=0.000000),IsValid=0),MBPOuterBounds=(Min=(X=0.000000,Y=0.000000,Z=0.000000),Max=(X=0.000000,Y=0.000000,Z=0.000000),IsValid=0),MBPNumSubdivs=2)
+ChaosSettings=(DefaultThreadingModel=TaskGraph,DedicatedThreadTickMode=VariableCappedWithTarget,DedicatedThreadBufferMode=Double)
+
+
diff --git a/Templates/TP_ArchvisBP/Config/DefaultGame.ini b/Templates/TP_ArchvisBP/Config/DefaultGame.ini
new file mode 100644
--- /dev/null
+++ b/Templates/TP_ArchvisBP/Config/DefaultGame.ini
@@ -0,0 +1,9 @@
+
+
+
+[/Script/EngineSettings.GeneralProjectSettings]
+ProjectID=A69F4BEA4684A29BD35BF7AB3A53907B
+
+[StartupActions]
+bAddPacks=False
+InsertPack=(PackSource="StarterContent.upack",PackName="StarterContent")
diff --git a/Templates/TP_ArchvisBP/Config/DefaultLightmass.ini b/Templates/TP_ArchvisBP/Config/DefaultLightmass.ini
new file mode 100644
--- /dev/null
+++ b/Templates/TP_ArchvisBP/Config/DefaultLightmass.ini
@@ -0,0 +1,22 @@
+; These are tweaked defaults for various lightmass solver and export settings
+; Artist oriented lightmass settings are in the editor UI
+; Documentation for all of these is in UnrealLightmass / Public / SceneExport.h
+; This ini is reloaded every time a lighting build begins, no need to restart
+
+; Warning: overwriting this file with an old version will cause the editor to crash.  This file must be in sync with the editor executable.
+; Instead, create a DefaultLightmass.ini in your project and override just the values you need, then the overrides will continue to work on version upgrades.
+; https://docs.unrealengine.com/latest/INT/Programming/Basics/ConfigurationFiles/
+;
+; For example, in your project's Config/DefaultLightmass.ini:
+; [DevOptions.PrecomputedDynamicObjectLighting]
+; SurfaceLightSampleSpacing=400
+
+[DevOptions.StaticLightingSettings]
+bCompressLightmaps=False
+NumIndirectLightingBounces=5
+NumSkyLightingBounces=5
+IndirectLightingSmoothness=0.9
+IndirectLightingQuality=2.0
+
+[DevOptions.StaticLightingSceneConstants]
+StaticLightingLevelScale=0.5
\ No newline at end of file
diff --git a/Templates/TP_ArchvisBP/Config/TemplateDefs.ini b/Templates/TP_ArchvisBP/Config/TemplateDefs.ini
new file mode 100644
--- /dev/null
+++ b/Templates/TP_ArchvisBP/Config/TemplateDefs.ini
@@ -0,0 +1,95 @@
+; This file specifies how a template will be treated when creating a project from it.
+; It includes rules for ignoring files, replacing files and changing the contents of files.
+; This file is annotated to help with creating your own templates.
+
+[/Script/GameProjectGeneration.TemplateProjectDefs]
+
+; The categories that this template should appear in. The values in this list refer to the Key value given in TemplateCategories.ini.
+Categories=AEC
+
+; The localized display names for this template. Each entry specifies a Language code and Unicode text to display in the New Project dialog.
+LocalizedDisplayNames=(Language="en", Text="Archvis")
+
+; The localized descriptions for this template, which is displayed when the template is selected.
+LocalizedDescriptions=(Language="en", Text="Template for architectural visualization projects.")
+; The sort key to use for the template to bypass the default alphabetical sorting.
+;SortKey="_1"
+
+; A list of folders to avoid copying when creating a project from this template. 
+; These should contain files only of concern to the template itself and not desired in the user's project.
+FoldersToIgnore=Media
+
+; The files to ignore copying, relative to the template root folder. 
+; %TEMPLATENAME% will be replaced by the template's folder name when resolving these. 
+; Filename comparisons are case-insensitive. 
+; For example, if this template is "TP_Blank" and the user creates "MyProject", then "TP_Blank.uproject" will be copied to "MyProject.uproject".
+FilesToIgnore="%TEMPLATENAME%.uproject"
+FilesToIgnore="%TEMPLATENAME%.png"
+FilesToIgnore="Config/TemplateDefs.ini"
+FilesToIgnore="Manifest.json"
+FilesToIgnore="contents.txt"
+
+; NOTE: For the purposes of the below examples, this template is considered to be named "TP_Blank", and the user has created a project called "MyProject".
+
+; Folders that will be renamed when creating a project from this template.
+; %PROJECTNAME% will be replaced by the project name specified by the user during creation when resolving these rules.
+; For example,  then "Source/TP_Blank" will be copied to "Source/MyProject".
+FolderRenames=(From="Source/%TEMPLATENAME%", To="Source/%PROJECTNAME%")
+FolderRenames=(From="Source/%TEMPLATENAME%Editor", To="Source/%PROJECTNAME%Editor")
+
+; Rules for renaming files when creating a project from this template. 
+; %TEMPLATENAME_UPPERCASE% and %PROJECTNAME_UPPERCASE% are the template and project names in upper case, eg. for "TP_Blank", its upper case representation is "TP_BLANK".
+; %TEMPLATENAME_LOWERCASE% and %PROJECTNAME_LOWERCASE% are the template and project names in lower case, eg. for "TP_Blank", its lower case representation is "tp_blank".
+; A rule does not have to include any of these modifiers. 
+; Extensions is a list of files to consider for these replacements. Eg. ("cpp", "h") will try to apply a rule only to files like "SomeFile.cpp" and "SomeFile.h".
+; bCaseSensitive specifies whether this replacement should be case-sensitive.
+
+; Here are some default rules for renaming files:
+
+; This rule states that a file called "TP_BLANK.cpp" will be renamed to "MYPROJECT.cpp".
+FilenameReplacements=(Extensions=("cpp","h","ini","cs"), From="%TEMPLATENAME_UPPERCASE%", To="%PROJECTNAME_UPPERCASE%", bCaseSensitive=true)
+
+; This rule states that a file called "tp_blank.ini" will be renamed to "myproject.ini"
+FilenameReplacements=(Extensions=("cpp","h","ini","cs"), From="%TEMPLATENAME_LOWERCASE%", To="%PROJECTNAME_LOWERCASE%", bCaseSensitive=true)
+
+; And this rule catches any irregularly-cased files, so that "Tp_bLANK.ini" will be renamed to "MyProject.ini".
+FilenameReplacements=(Extensions=("cpp","h","ini","cs"), From="%TEMPLATENAME%",To="%PROJECTNAME%",bCaseSensitive=false)
+
+; Rules for renaming strings in the contents of files. 
+; The same rules as above apply, but any usages of strings that match these rules inside the file will be changed. This includes all substrings!
+; For example, "TP_Blank" will be replaced by "MyProject", but "TP_Blank_Name" will also be replaced by "MyProject_Name".
+ReplacementsInFiles=(Extensions=("cpp","h","ini","cs"), From="%TEMPLATENAME_UPPERCASE%", To="%PROJECTNAME_UPPERCASE%", bCaseSensitive=true)
+ReplacementsInFiles=(Extensions=("cpp","h","ini","cs"), From="%TEMPLATENAME_LOWERCASE%", To="%PROJECTNAME_LOWERCASE%", bCaseSensitive=true)
+ReplacementsInFiles=(Extensions=("cpp","h","ini","cs"), From="%TEMPLATENAME%", To="%PROJECTNAME%", bCaseSensitive=false)
+
+; The following settings are not enabled for this project, but are included for completeness:
+
+; Whether or not this template can be used to create a project. Default: true
+; If this is set to false, then the template becomes a Feature Pack that can only be referenced from other templates, or added to a project from the Content Browser.
+;bAllowProjectCreation=false
+
+; Which settings on the project settings page should be hidden for this template. Default: none
+; Allowed values are:
+; All - Hides the settings page entirely
+; Languages - Choose between C++ and Blueprint versions.
+; HardwareTarget - Choose between Mobile and Desktop hardware targets.
+; GraphicsPreset - Choose between Maximum and Scalable graphics.
+; StarterContent - Choose whether to add starter content to the project.
+; VR - Choose whether to enable Virtual Reality support for the project.
+;HiddenSettings=All
+HiddenSettings=Languages
+HiddenSettings=HardwareTarget
+HiddenSettings=GraphicsPreset 
+;HiddenSettings=StarterContent 
+;HiddenSettings=XR
+
+; The class types that this project uses.
+; This is only for display in the template description, and does not affect project creation.
+;ClassTypes=GameMode, Character, Actor
+
+; The asset types that this project uses.
+; This is only for display in the template description, and does not affect project creation.
+;AssetTypes=Skeleton, Static Mesh, Texture
+
+; Shared feature packs. The files in these packs listed in these structures marked as 'additionalfiles' will be copied on project generation.
+;SharedContentPacks=(MountName="2DSideScroller",DetailLevels=("Standard"))
diff --git a/Templates/TP_ArchvisBP/TP_ArchvisBP.uproject b/Templates/TP_ArchvisBP/TP_ArchvisBP.uproject
new file mode 100644
--- /dev/null
+++ b/Templates/TP_ArchvisBP/TP_ArchvisBP.uproject
@@ -0,0 +1,54 @@
+{
+	"FileVersion": 3,
+	"EngineAssociation": "",
+	"Category": "",
+	"Description": "",
+	"Plugins": [
+		{
+			"Name": "SunPosition",
+			"Enabled": true
+		},
+		{
+			"Name": "DataPrepEditor",
+			"Enabled": true
+		},
+		{
+			"Name": "DatasmithImporter",
+			"Enabled": true
+		},
+		{
+			"Name": "VariantManager",
+			"Enabled": true
+		},
+		{
+			"Name": "PythonScriptPlugin",
+			"Enabled": true
+		},
+		{
+			"Name": "HDRIBackdrop",
+			"Enabled": true
+		},
+		{
+			"Name": "SequencerScripting",
+			"Enabled": false
+		},
+		{
+			"Name": "DatasmithCADImporter",
+			"Enabled": true,
+			"SupportedTargetPlatforms": [
+				"Win64"
+			]
+		},
+		{
+			"Name": "DatasmithIFCImporter",
+			"Enabled": true,
+			"SupportedTargetPlatforms": [
+				"Win64"
+			]
+		},
+		{
+			"Name": "StaticMeshEditorExtension",
+			"Enabled": true
+		}
+	]
+}
\ No newline at end of file
