Refactoring the ArchvisBP template to reflect recent decisions on folder structure.
Archvis template no loger has LevelTemplates
Renamed assets to follow general guidelines
Added thumbnail to the "welcome" / main level

Tested by launching unreal, making a new project with this template.  everything plays correctly.

Next: add NPR levels

diff --git a/Templates/TP_ArchvisBP/Config/DefaultEngine.ini b/Templates/TP_ArchvisBP/Config/DefaultEngine.ini
--- a/Templates/TP_ArchvisBP/Config/DefaultEngine.ini
+++ b/Templates/TP_ArchvisBP/Config/DefaultEngine.ini
@@ -39,55 +39,4 @@ GlobalDefaultGameMode=/Game/ArchvisProject/ArchvisGameMode.ArchvisGameMode_C
 -TemplateMapInfos=(ThumbnailTexture=Texture2D'/Engine/Maps/Templates/Thumbnails/Default.Default',Map="/Engine/Maps/Templates/Template_Default")
 -TemplateMapInfos=(ThumbnailTexture=Texture2D'/Engine/Maps/Templates/Thumbnails/TimeOfDay.TimeOfDay',Map="/Engine/Maps/Templates/TimeOfDay_Default")
 -TemplateMapInfos=(ThumbnailTexture=Texture2D'/Engine/Maps/Templates/Thumbnails/VR-Basic.VR-Basic',Map="/Engine/Maps/Templates/VR-Basic")
-
-; Add Predefined Template Maps that are suitable for this template
-;+TemplateMapInfos=(ThumbnailTexture=Texture2D'/Game/Developers/Maps/Thumbnails/AssetName.AssetName',Map="/Game/Maps/Templates/MapName")
-+TemplateMapInfos=(ThumbnailTexture=Texture2D'/Game/LevelTemplates/Default/Thumbnails/Default.Default',Map="/Game/LevelTemplates/Default/Default")
-
-[/Script/Engine.PhysicsSettings]
-DefaultGravityZ=-980.000000
-DefaultTerminalVelocity=4000.000000
-DefaultFluidFriction=0.300000
-SimulateScratchMemorySize=262144
-RagdollAggregateThreshold=4
-TriangleMeshTriangleMinAreaThreshold=5.000000
-bEnableShapeSharing=False
-bEnablePCM=True
-bEnableStabilization=False
-bWarnMissingLocks=True
-bEnable2DPhysics=False
-PhysicErrorCorrection=(PingExtrapolation=0.100000,PingLimit=100.000000,ErrorPerLinearDifference=1.000000,ErrorPerAngularDifference=1.000000,MaxRestoredStateError=1.000000,MaxLinearHardSnapDistance=400.000000,PositionLerp=0.000000,AngleLerp=0.400000,LinearVelocityCoefficient=100.000000,AngularVelocityCoefficient=10.000000,ErrorAccumulationSeconds=0.500000,ErrorAccumulationDistanceSq=15.000000,ErrorAccumulationSimilarity=100.000000)
-LockedAxis=Invalid
-DefaultDegreesOfFreedom=Full3D
-BounceThresholdVelocity=200.000000
-FrictionCombineMode=Average
-RestitutionCombineMode=Average
-MaxAngularVelocity=3600.000000
-MaxDepenetrationVelocity=0.000000
-ContactOffsetMultiplier=0.020000
-MinContactOffset=2.000000
-MaxContactOffset=8.000000
-bSimulateSkeletalMeshOnDedicatedServer=True
-DefaultShapeComplexity=CTF_UseSimpleAndComplex
-bDefaultHasComplexCollision=True
-bSuppressFaceRemapTable=False
-bSupportUVFromHitResults=False
-bDisableActiveActors=False
-bDisableKinematicStaticPairs=False
-bDisableKinematicKinematicPairs=False
-bDisableCCD=False
-bEnableEnhancedDeterminism=False
-AnimPhysicsMinDeltaTime=0.000000
-bSimulateAnimPhysicsAfterReset=False
-MaxPhysicsDeltaTime=0.033333
-bSubstepping=False
-bSubsteppingAsync=False
-MaxSubstepDeltaTime=0.016667
-MaxSubsteps=6
-SyncSceneSmoothingFactor=0.000000
-InitialAverageFrameRate=0.016667
-PhysXTreeRebuildRate=10
-DefaultBroadphaseSettings=(bUseMBPOnClient=False,bUseMBPOnServer=False,bUseMBPOuterBounds=False,MBPBounds=(Min=(X=0.000000,Y=0.000000,Z=0.000000),Max=(X=0.000000,Y=0.000000,Z=0.000000),IsValid=0),MBPOuterBounds=(Min=(X=0.000000,Y=0.000000,Z=0.000000),Max=(X=0.000000,Y=0.000000,Z=0.000000),IsValid=0),MBPNumSubdivs=2)
-ChaosSettings=(DefaultThreadingModel=TaskGraph,DedicatedThreadTickMode=VariableCappedWithTarget,DedicatedThreadBufferMode=Double)
-
-
++TemplateMapInfos=(ThumbnailTexture=Texture2D'/Engine/Maps/Templates/Thumbnails/Basic.Basic',Map="/Engine/Maps/Templates/Basic")
\ No newline at end of file
