AEC Templates level templates tweaks

diff --git a/Templates/TP_AEC_BlankBP/Config/DefaultEngine.ini b/Templates/TP_AEC_BlankBP/Config/DefaultEngine.ini
--- a/Templates/TP_AEC_BlankBP/Config/DefaultEngine.ini
+++ b/Templates/TP_AEC_BlankBP/Config/DefaultEngine.ini
@@ -37,6 +37,5 @@ EditorStartupMap=/Game/Main.Main
 [/Script/UnrealEd.UnrealEdEngine]
 ; Remove Engine Template maps
 -TemplateMapInfos=(ThumbnailTexture=Texture2D'/Engine/Maps/Templates/Thumbnails/Default.Default',Map="/Engine/Maps/Templates/Template_Default")
--TemplateMapInfos=(ThumbnailTexture=Texture2D'/Engine/Maps/Templates/Thumbnails/TimeOfDay.TimeOfDay',Map="/Engine/Maps/Templates/TimeOfDay_Default")
--TemplateMapInfos=(ThumbnailTexture=Texture2D'/Engine/Maps/Templates/Thumbnails/VR-Basic.VR-Basic',Map="/Engine/Maps/Templates/VR-Basic")
-+TemplateMapInfos=(ThumbnailTexture=Texture2D'/Engine/Maps/Templates/Thumbnails/Basic.Basic',Map="/Engine/Maps/Templates/Basic")
\ No newline at end of file
++TemplateMapInfos=(ThumbnailTexture=Texture2D'/Engine/Maps/Templates/Thumbnails/TimeOfDay.TimeOfDay',Map="/Engine/Maps/Templates/TimeOfDay_Default")
+-TemplateMapInfos=(ThumbnailTexture=Texture2D'/Engine/Maps/Templates/Thumbnails/VR-Basic.VR-Basic',Map="/Engine/Maps/Templates/VR-Basic")
\ No newline at end of file
diff --git a/Templates/TP_ArchvisBP/Config/DefaultEngine.ini b/Templates/TP_ArchvisBP/Config/DefaultEngine.ini
--- a/Templates/TP_ArchvisBP/Config/DefaultEngine.ini
+++ b/Templates/TP_ArchvisBP/Config/DefaultEngine.ini
@@ -30,13 +30,12 @@ r.NormalMapsForStaticLighting=False
 
 [/Script/EngineSettings.GameMapsSettings]
 ;defines what level to be loaded by default in the editor and at playback time
-GameDefaultMap=/Game/ArchvisProject/Main.Main
-EditorStartupMap=/Game/ArchvisProject/Main.Main
+GameDefaultMap=/Game/ArchvisProject/Exterior.Exterior
+EditorStartupMap=/Game/ArchvisProject/Exterior.Exterior
 GlobalDefaultGameMode=/Game/ArchvisProject/ArchvisGameMode.ArchvisGameMode_C
 
 [/Script/UnrealEd.UnrealEdEngine]
 ; Remove Engine Template maps
 -TemplateMapInfos=(ThumbnailTexture=Texture2D'/Engine/Maps/Templates/Thumbnails/Default.Default',Map="/Engine/Maps/Templates/Template_Default")
--TemplateMapInfos=(ThumbnailTexture=Texture2D'/Engine/Maps/Templates/Thumbnails/TimeOfDay.TimeOfDay',Map="/Engine/Maps/Templates/TimeOfDay_Default")
++TemplateMapInfos=(ThumbnailTexture=Texture2D'/Engine/Maps/Templates/Thumbnails/TimeOfDay.TimeOfDay',Map="/Engine/Maps/Templates/TimeOfDay_Default")
 -TemplateMapInfos=(ThumbnailTexture=Texture2D'/Engine/Maps/Templates/Thumbnails/VR-Basic.VR-Basic',Map="/Engine/Maps/Templates/VR-Basic")
-+TemplateMapInfos=(ThumbnailTexture=Texture2D'/Engine/Maps/Templates/Thumbnails/Basic.Basic',Map="/Engine/Maps/Templates/Basic")
\ No newline at end of file
