Archvis template:

fixing icons and thumbnails / image previews
Improving lightmass settings
solving colors on silhouette objects
updating project description

diff --git a/Templates/TP_ArchvisBP/Config/TemplateDefs.ini b/Templates/TP_ArchvisBP/Config/TemplateDefs.ini
--- a/Templates/TP_ArchvisBP/Config/TemplateDefs.ini
+++ b/Templates/TP_ArchvisBP/Config/TemplateDefs.ini
@@ -11,9 +11,9 @@ Categories=AEC
 LocalizedDisplayNames=(Language="en", Text="Archvis")
 
 ; The localized descriptions for this template, which is displayed when the template is selected.
-LocalizedDescriptions=(Language="en", Text="Template for architectural visualization projects.")
+LocalizedDescriptions=(Language="en", Text="This template is dedicated to architectural visualization workflows with examples for sun studies, interior renderings, and non-photorealistic stylized renderings.\n\nThis template enables you to use Datasmith to import models from various design applications and easily light them.")
 ; The sort key to use for the template to bypass the default alphabetical sorting.
-;SortKey="_1"
+SortKey="1"
 
 ; A list of folders to avoid copying when creating a project from this template. 
 ; These should contain files only of concern to the template itself and not desired in the user's project.
